Updated template description translations (Korean)

diff --git a/Templates/TP_FirstPerson/manifest.json b/Templates/TP_FirstPerson/manifest.json
--- a/Templates/TP_FirstPerson/manifest.json
+++ b/Templates/TP_FirstPerson/manifest.json
@@ -27,7 +27,7 @@
 		},
 		{
 			"Language": "ko",
-			"Text": "일인칭 템플릿은 일인칭 시점에서 보이는 (양팔로 표현되는) 플레이어 캐릭터를 특징으로 합니다. 캐릭터 이동은 키보드, 컨트롤러, 터치 디바이스의 가상 조이스틱으로 가능합니다. 추가적으로 플레이어는 마우스, 컨트롤러, 터치 디바이스의 가상 조이스틱을 사용하여 주변을 둘러볼 수도 있습니다. 캐릭터에는 총이 장착되어 있어 마우스, 컨트롤러, 터치 디바이스의 가상 조이스틱을 사용하면 간단한 발사체가 발사되어, 아레나 벽에 튕기면서 레벨에 있는 몇몇 물리 오브젝트에 영향을 줍니다."
+			"Text": "일인칭 템플릿은 일인칭 시점에서 보이는 (양팔로 표현되는) 플레이어 캐릭터를 특징으로 합니다. 캐릭터는 레벨 내에서 키보드, 컨트롤러나 터치 디바이스의 가상 조이스틱으로 이동할 수 있습니다. 추가적으로 플레이어는 마우스, 컨트롤러나 터치 디바이스의 가상 조이스틱을 사용하여 주변을 둘러볼 수도 있습니다. 또한, 캐릭터는 총을 장착할 수 있으며 마우스, 컨트롤러나 터치 디바이스의 가상 조이스틱을 사용해 간단한 구형 발사체를 발사하여, 아레나 벽에 튕기면서 레벨에 있는 일부 피직스 오브젝트에 영향을 줍니다."
 		},
 		{
 			"Language": "zh-Hans",
